SWidget Integration Test and Unit Test Framework Setup

diff --git a/Engine/Plugins/Tests/WidgetAutomationTests/Resources/TestBaselineFiles/UIAutomation_SlateTestTextBlock_ColorAndOpacity.json b/Engine/Plugins/Tests/WidgetAutomationTests/Resources/TestBaselineFiles/UIAutomation_SlateTestTextBlock_ColorAndOpacity.json
new file mode 100644
--- /dev/null
+++ b/Engine/Plugins/Tests/WidgetAutomationTests/Resources/TestBaselineFiles/UIAutomation_SlateTestTextBlock_ColorAndOpacity.json
@@ -0,0 +1,16 @@
+{
+	"ComputeDesiredSizeX": 75,
+	"ComputeDesiredSizeY": 15,
+	"AllottedGeometry": "[Abs=X=3.000 Y=37.000, Scale=1.00, Size=X=1000.000 Y=700.000]",
+	"WidgetStyle":
+	{
+		"ColorAndOpacityTint": "(R=1.000000,G=1.000000,B=1.000000,A=1.000000)",
+		"ForegroundColor": "(R=0.527115,G=0.527115,B=0.527115,A=1.000000)",
+		"SubduedForeground": "(R=0.527115,G=0.527115,B=0.527115,A=0.600000)"
+	},
+	"ShapedTextPayload":
+	{
+		"Tint": "(R=0.500000,G=0.400000,B=0.300000,A=0.900000)",
+		"OutlineTint": "(R=0.000000,G=0.000000,B=0.000000,A=1.000000)"
+	}
+}
\ No newline at end of file
diff --git a/Engine/Plugins/Tests/WidgetAutomationTests/Source/WidgetAutomationTests/Private/FocusTest.cpp b/Engine/Plugins/Tests/WidgetAutomationTests/Source/WidgetAutomationTests/Private/FocusTest.cpp
new file mode 100644
--- /dev/null
+++ b/Engine/Plugins/Tests/WidgetAutomationTests/Source/WidgetAutomationTests/Private/FocusTest.cpp
@@ -0,0 +1,92 @@
+// Copyright Epic Games, Inc. All Rights Reserved.
+
+#include "CoreMinimal.h"
+#include "Framework/Application/SlateApplication.h"
+#include "Misc/AutomationTest.h"
+#include "TestDataAndFunctionNames.h"
+#include "SlateTestBase.h"
+#include "SWidgetMock.h"
+#include "Widgets/Input/SButton.h"
+#include "Widgets/SWindow.h"
+
+#if WITH_DEV_AUTOMATION_TESTS
+
+namespace UE::SlateWidgetAutomationTest
+{
+IMPLEMENT_SIMPLE_AUTOMATION_TEST(FFocusTest, "Slate.FocusTest", EAutomationTestFlags::ApplicationContextMask | EAutomationTestFlags::EngineFilter)
+
+template<class T>
+class FSlateTestFocus : public FSlateTestBase<T>
+{
+public:
+
+	using FSlateTestBase<T>::GetSubjectWidget;
+	using FSlateTestBase<T>::GetMainWindow;
+	using FSlateTestBase<T>::CreateWidget;
+
+	void Setup(FAutomationTestBase* TestObj) override
+	{
+		// 1. In the main window, insert a panel widget containing two widgets of the same type that support focus.
+		// We call the two widgets SubjectWidget and SiblingWidget. SubjectWidget is the target of focus in the test.
+		PanelWidget = SNew(SWidgetMock<SVerticalBox>);
+		SiblingWidget = CreateWidget();
+
+		PanelWidget->AddSlot()[GetSubjectWidget()];
+		PanelWidget->AddSlot()[SiblingWidget.ToSharedRef()];
+
+		GetMainWindow()->SetContent(PanelWidget.ToSharedRef());
+
+		// 2. Initially set focus to the sibling widget. 
+		FSlateApplication::Get().SetUserFocus(0, SiblingWidget);
+
+		// 3. Erase any side effect that we might have caused as a result of setting up the test.
+		GetSubjectWidget()->ClearFunctionRecords();
+		SiblingWidget->ClearFunctionRecords();
+
+	}
+
+	void Run(FAutomationTestBase* TestObj) override
+	{
+		// Manually set the focus on the subject widget.
+		FSlateApplication::Get().SetUserFocus(0, GetSubjectWidget());
+	}
+
+	void Validate(FAutomationTestBase* TestObj ) override
+	{
+		// Verify the expected callbacks were called on the right widgets.
+		TestObj->AddErrorIfFalse(GetSubjectWidget()->GetValueInFunctionCalls(FTestFunctionNames::NAME_OnFocusReceived) == 1, FString::Printf(TEXT("In %s: OnFocusReceived callback was not called on the widget receiving focus."), GetGeneratedTypeName<T>()));
+		TestObj->AddErrorIfFalse(GetSubjectWidget()->GetValueInFunctionCalls(FTestFunctionNames::NAME_OnFocusChanging) == 1, FString::Printf(TEXT("In %s: OnFocusChanging callback was not called on the widget receiving focus."), GetGeneratedTypeName<T>()));
+		TestObj->AddErrorIfFalse(SiblingWidget->GetValueInFunctionCalls(FTestFunctionNames::NAME_OnFocusLost) == 1, FString::Printf(TEXT("In %s: OnFocusLost callback was not called on the widget initally having focus."), GetGeneratedTypeName<T>()));
+		TestObj->AddErrorIfFalse(SiblingWidget->GetValueInFunctionCalls(FTestFunctionNames::NAME_OnFocusChanging) == 1, FString::Printf(TEXT("In %s: OnFocusChanging callback was not called on the widget initally having focus."), GetGeneratedTypeName<T>()));
+
+		// Verify that unexpected callbacks were NOT called on the widgets.
+		TestObj->AddErrorIfFalse(GetSubjectWidget()->GetValueInFunctionCalls(FTestFunctionNames::NAME_OnFocusLost) == 0, FString::Printf(TEXT("In %s: OnFocusLost callback was called the widget receiving focus."), GetGeneratedTypeName<T>()));
+		TestObj->AddErrorIfFalse(SiblingWidget->GetValueInFunctionCalls(FTestFunctionNames::NAME_OnFocusReceived) == 0, FString::Printf(TEXT("In %s: OnFocusReceived callback was called on the widget initally having focus."), GetGeneratedTypeName<T>()));
+
+		TestObj->AddErrorIfFalse(GetSubjectWidget()->HasAnyUserFocus().IsSet(), FString::Printf(TEXT("In %s: No user is focusing on the widget receiving focus."), GetGeneratedTypeName<T>()));
+		TestObj->AddErrorIfFalse(GetMainWindow()->HasFocusedDescendants(), FString::Printf(TEXT("In %s: The parent window does not detect that its descendant has focus because it detects no user that is focusing on it."), GetGeneratedTypeName<T>()));
+		TestObj->AddErrorIfFalse(PanelWidget->HasFocusedDescendants(), FString::Printf(TEXT("In %s: The immediate parent of the widget does not detect that its descendant has focus because it detects no user that is focusing on it."), GetGeneratedTypeName<T>()));
+	}
+
+	virtual void TearDown(FAutomationTestBase* TestObj)
+	{
+		FSlateApplication::Get().RequestDestroyWindow(GetMainWindow());
+	}
+
+private:
+	TSharedPtr<SWidgetMock<T>> SiblingWidget;
+	TSharedPtr<SWidgetMock<SVerticalBox>> PanelWidget;
+};
+
+bool FFocusTest::RunTest(const FString& Parameters)
+{
+	FSlateTestFocus<SButton> MyTest = FSlateTestFocus<SButton>();
+
+	MyTest.Setup(this);
+	MyTest.Run(this);
+	MyTest.Validate(this);
+	MyTest.TearDown(this);
+	return true;
+}
+}
+#endif
\ No newline at end of file
diff --git a/Engine/Plugins/Tests/WidgetAutomationTests/Source/WidgetAutomationTests/Private/STextBlockTest.cpp b/Engine/Plugins/Tests/WidgetAutomationTests/Source/WidgetAutomationTests/Private/STextBlockTest.cpp
new file mode 100644
--- /dev/null
+++ b/Engine/Plugins/Tests/WidgetAutomationTests/Source/WidgetAutomationTests/Private/STextBlockTest.cpp
@@ -0,0 +1,122 @@
+// Copyright Epic Games, Inc. All Rights Reserved.
+
+#include "CoreMinimal.h"
+#include "Misc/AutomationTest.h"
+#include "SlateTestBase.h"
+#include "SWidgetMock.h"
+
+#include "Widgets/Text/STextBlock.h"
+#include "Widgets/SWindow.h"
+#include "Framework/Application/SlateApplication.h"
+
+#include "Dom/JsonObject.h"
+
+#if WITH_DEV_AUTOMATION_TESTS
+
+namespace UE::SlateWidgetAutomationTest
+{
+IMPLEMENT_SIMPLE_AUTOMATION_TEST(FTextBlockTest, "Slate.STextBlockTest", EAutomationTestFlags::ApplicationContextMask | EAutomationTestFlags::EngineFilter)
+
+//@TODO Should we create a new class per property? We are going to need at least different RunTest and ValidateTest functions for each property of textblock.
+class FSlateTestTextBlock : public FSlateTestBase<STextBlock>
+{
+
+public:
+
+	FSlateTestTextBlock(FString InTestName) : FSlateTestBase<STextBlock>(MoveTemp(InTestName)) {};
+
+	void Setup(FAutomationTestBase* TestObj) override
+	{
+		// 1. In the main window, insert a panel widget containing the TextBlock widget being tested.
+		TSharedRef<SVerticalBox> VerticalBox = SNew(SVerticalBox);
+		VerticalBox->AddSlot()[GetSubjectWidget()];
+		GetMainWindow()->SetContent(VerticalBox);
+
+		// 2. Initialize the textbox with a sample text.
+		FString TestText = "ExampleText";
+		GetSubjectWidget()->SetText(FText::FromString(TestText));
+
+		// 3. Enable fast update on the main window so that we can test global invalidation on its content.
+		GetMainWindow()->SetAllowFastUpdate(true);
+
+		// 4. Erase any side effect that we might have caused as a result of setting up the test.
+		// This must be delayed to the next frame because the side effects of SetAllowFastUpdate will appear next frame.
+		TestObj->AddCommand(new FDelayedFunctionLatentCommand([this, TestObj]()
+			{
+				GetSubjectWidget()->ClearFunctionRecords();
+
+				// 5. Load the JSON object for this test. The data will be used as the ground truth of the test.
+				// It is important to call this function after ClearFunctionRecords(), otherwise we may end up clearing the errors we faced during file I/O.
+				GetSubjectWidget()->FindOrCreateTestJSONFile(GetTestName());
+			}, 0.1f));
+	}
+
+	void Run(FAutomationTestBase* TestObj) override
+	{
+		GetSubjectWidget()->SetColorAndOpacity(FSlateColor(FLinearColor(0.5,0.4,0.3,0.9)));
+	}
+
+	void Validate(FAutomationTestBase* TestObj) override
+	{
+		// Ensure the JSON file for the test exists and is loaded. We don't want to continue verifications without it.
+		// This should be the first thing we validate in the tests that use a JSON file.
+		if (!JSONValidation(TestObj, "ColorAndOpacity"))
+		{
+			return;
+		}
+
+		// generic tests for property change.
+		ValidatePropertyChange<FSlateColor>(TestObj, &STextBlock::GetColorAndOpacity, "ColorAndOpacity", FSlateColor(FLinearColor(0.5, 0.4, 0.3, 0.9)), EInvalidateWidgetReason::Layout|EInvalidateWidgetReason::Paint);
+
+		// Add any non-generic tests here.
+
+		TestObj->AddCommand(new FDelayedFunctionLatentCommand([this, TestObj]()
+			{
+				if (!GetMainWindow()->Advanced_IsInvalidationRoot())
+				{
+					// Confirm that the payload tint calculated in OnPaint matches the color we set in the test.
+					if (TSharedPtr<FJsonObject> JsonObject = GetSubjectWidget()->GetTestJSONObject())
+					{
+						FString PayloadTint = JsonObject->GetObjectField("ShapedTextPayload")->GetStringField("Tint");
+						TestObj->AddErrorIfFalse(PayloadTint == FLinearColor(0.5, 0.4, 0.3, 0.9).ToString(), FString::Printf(TEXT("In %s: The tint of the payload does not match the passed color. Expected: %s, Actual: %s"), GetGeneratedTypeName<STextBlock>(), *FLinearColor(0.5, 0.4, 0.3, 0.9).ToString(), *PayloadTint));
+					}
+				}
+			}, 0.1f));
+	}
+
+	void TearDown(FAutomationTestBase* TestObj)
+	{
+		if (GetSubjectWidget()->GetTestJSONObjectToWrite().IsValid() && !FSlateTestHelper::WriteJSONToTxt(GetSubjectWidget()->GetTestJSONObjectToWrite().ToSharedRef(), GetTestName()))
+		{
+			TestObj->AddError(FString::Printf(TEXT("Failed to write the ground truth JSON file for the test %s."), "ColorAndOpacity", *GetTestName()));
+		}
+		FSlateApplication::Get().RequestDestroyWindow(GetMainWindow());
+	}
+};
+
+bool FTextBlockTest::RunTest(const FString& Parameters)
+{
+	TSharedPtr<FSlateTestTextBlock> MyTest = MakeShared<FSlateTestTextBlock>(TEXT("SlateTestTextBlock_ColorAndOpacity"));
+
+	MyTest->Setup(this);
+
+	// Wait for the next frame so that SetupTest has finished.
+	// Whether we delay at each of these steps or not depends on the nature of the test. 
+	// But once a step has been delayed, all the others must be too.
+	AddCommand(new FDelayedFunctionLatentCommand([this, MyTest]()
+		{
+			MyTest->Run(this);
+			MyTest->Validate(this);
+
+			// Wait for ValidateTest to finish.
+			this->AddCommand(new FDelayedFunctionLatentCommand([this, MyTest]()
+				{
+					MyTest->TearDown(this);
+				}, 0.1f));
+
+		}, 0.1f));
+
+	return true;
+}
+}
+#endif
\ No newline at end of file
diff --git a/Engine/Plugins/Tests/WidgetAutomationTests/Source/WidgetAutomationTests/Private/SlateTestHelper.cpp b/Engine/Plugins/Tests/WidgetAutomationTests/Source/WidgetAutomationTests/Private/SlateTestHelper.cpp
new file mode 100644
--- /dev/null
+++ b/Engine/Plugins/Tests/WidgetAutomationTests/Source/WidgetAutomationTests/Private/SlateTestHelper.cpp
@@ -0,0 +1,70 @@
+// Copyright Epic Games, Inc. All Rights Reserved.
+
+#include "SlateTestHelper.h"
+#include "CoreMinimal.h"
+#include "Rendering/DrawElementPayloads.h"
+#include "Serialization/JsonReader.h"
+#include "Serialization/JsonSerializer.h"
+#include "Misc/FileHelper.h"
+#include "Misc/Paths.h"
+
+#if WITH_DEV_AUTOMATION_TESTS
+namespace UE::SlateWidgetAutomationTest
+{
+bool FSlateTestHelper::WriteJSONToTxt(TSharedRef<FJsonObject> ObjToWrite, const FString& TestName)
+{
+	FString JsonString;
+	TSharedRef<TJsonWriter<>> JsonWriter = TJsonWriterFactory<>::Create(&JsonString);
+	if (FJsonSerializer::Serialize(ObjToWrite, JsonWriter, true))
+	{
+		return FFileHelper::SaveStringToFile(JsonString, *FSlateTestHelper::GetFullPath(TestName));
+
+	}
+	return false;
+}
+
+TSharedPtr<FJsonObject> FSlateTestHelper::LoadJSONFromTxt(const FString& TestName)
+{
+	FString JsonString;
+	TSharedPtr<FJsonObject> OutJsonObject;
+
+	//load file and convert to JSON
+	if (FFileHelper::LoadFileToString(JsonString, *FSlateTestHelper::GetFullPath(TestName)))
+	{
+		TSharedRef<TJsonReader<>> JsonReader = TJsonReaderFactory<>::Create(JsonString);
+		FJsonSerializer::Deserialize(JsonReader, OutJsonObject);
+	}
+
+	return OutJsonObject;
+}
+
+FString FSlateTestHelper::GetFullPath(FStringView TestName)
+{
+	FString FileLocation = FPaths::ConvertRelativePathToFull(FPaths::EnginePluginsDir()) + TEXT("Tests/WidgetAutomationTests/Resources/TestBaselineFiles/");
+	FString FullPath = FString::Printf(TEXT("%sUIAutomation_%s.json"), *FileLocation, TestName.GetData());
+	return FullPath;
+}
+
+void FSlateTestHelper::ToJson(const FGeometry& MyGeometry, TSharedPtr<FJsonObject> MyJsonObject)
+{
+	MyJsonObject->SetStringField("AllottedGeometry", MyGeometry.ToString());
+}
+
+void FSlateTestHelper::ToJson(const FWidgetStyle& InWidgetStyle, TSharedPtr<FJsonObject> MyJsonObject)
+{
+	TSharedPtr<FJsonObject> InnerJsonObject = MakeShared<FJsonObject>();
+	InnerJsonObject->SetStringField("ColorAndOpacityTint", InWidgetStyle.GetColorAndOpacityTint().ToString());
+	InnerJsonObject->SetStringField("ForegroundColor", InWidgetStyle.GetForegroundColor().ToString());
+	InnerJsonObject->SetStringField("SubduedForeground", InWidgetStyle.GetSubduedForegroundColor().ToString());
+	MyJsonObject->SetObjectField("WidgetStyle", InnerJsonObject);
+}
+
+void FSlateTestHelper::ToJson(const FSlateShapedTextPayload& ShapedTextPayload, TSharedPtr<FJsonObject> MyJsonObject)
+{
+	TSharedPtr<FJsonObject> InnerJsonObject = MakeShared<FJsonObject>();
+	InnerJsonObject->SetStringField("Tint", ShapedTextPayload.GetTint().ToString());
+	InnerJsonObject->SetStringField("OutlineTint", ShapedTextPayload.GetOutlineTint().ToString());
+	MyJsonObject->SetObjectField("ShapedTextPayload", InnerJsonObject);
+}
+}
+#endif
\ No newline at end of file
diff --git a/Engine/Plugins/Tests/WidgetAutomationTests/Source/WidgetAutomationTests/Private/TestDataAndFunctionNames.cpp b/Engine/Plugins/Tests/WidgetAutomationTests/Source/WidgetAutomationTests/Private/TestDataAndFunctionNames.cpp
new file mode 100644
--- /dev/null
+++ b/Engine/Plugins/Tests/WidgetAutomationTests/Source/WidgetAutomationTests/Private/TestDataAndFunctionNames.cpp
@@ -0,0 +1,23 @@
+#include "TestDataAndFunctionNames.h"
+
+#if WITH_DEV_AUTOMATION_TESTS
+
+namespace UE::SlateWidgetAutomationTest
+{
+	namespace FTestIONames
+	{
+		const FLazyName NAME_FailedToLoadJSON = "FailedToLoadJSON";
+		const FLazyName NAME_TestFileNotFound = "TestFileNotFound";
+		const FLazyName NAME_BadTestFile = "BadTestFile";
+	};
+
+	namespace FTestFunctionNames
+	{
+		const FLazyName NAME_OnPaint = "OnPaint";
+		const FLazyName NAME_ComputeDesiredSize = "ComputeDesiredSize";
+		const FLazyName NAME_OnFocusChanging = "OnFocusChanging";
+		const FLazyName NAME_OnFocusLost = "OnFocusLost";
+		const FLazyName NAME_OnFocusReceived = "OnFocusReceived";
+	};
+}
+#endif
\ No newline at end of file
diff --git a/Engine/Plugins/Tests/WidgetAutomationTests/Source/WidgetAutomationTests/Private/WidgetAutomationTestsModule.cpp b/Engine/Plugins/Tests/WidgetAutomationTests/Source/WidgetAutomationTests/Private/WidgetAutomationTestsModule.cpp
new file mode 100644
--- /dev/null
+++ b/Engine/Plugins/Tests/WidgetAutomationTests/Source/WidgetAutomationTests/Private/WidgetAutomationTestsModule.cpp
@@ -0,0 +1,23 @@
+// Copyright Epic Games, Inc. All Rights Reserved.
+
+#include "IWidgetAutomationTestsModule.h"
+
+class FWidgetAutomationTests : public IWidgetAutomationTests
+{
+	/** IModuleInterface implementation */
+	virtual void StartupModule() override;
+	virtual void ShutdownModule() override;
+};
+
+IMPLEMENT_MODULE(IWidgetAutomationTests, WidgetAutomationTests)
+
+void FWidgetAutomationTests::StartupModule()
+{
+	// This code will execute after your module is loaded into memory (but after global variables are initialized, of course.)
+}
+
+void FWidgetAutomationTests::ShutdownModule()
+{
+	// This function may be called during shutdown to clean up your module.  For modules that support dynamic reloading,
+	// we call this function before unloading the module.
+}
\ No newline at end of file
diff --git a/Engine/Plugins/Tests/WidgetAutomationTests/Source/WidgetAutomationTests/Private/WidgetMockNonTemplate.cpp b/Engine/Plugins/Tests/WidgetAutomationTests/Source/WidgetAutomationTests/Private/WidgetMockNonTemplate.cpp
new file mode 100644
--- /dev/null
+++ b/Engine/Plugins/Tests/WidgetAutomationTests/Source/WidgetAutomationTests/Private/WidgetMockNonTemplate.cpp
@@ -0,0 +1,202 @@
+// Copyright Epic Games, Inc. All Rights Reserved.
+#include "WidgetMockNonTemplate.h"
+#include "Dom/JsonObject.h"
+#include "Misc/GeneratedTypeName.h"
+#include "Rendering/DrawElementPayloads.h"
+#include "Serialization/JsonReader.h"
+#include "SlateTestHelper.h"
+#include "TestDataAndFunctionNames.h"
+#include "Widgets/InvalidateWidgetReason.h"
+#include "Widgets/SWindow.h"
+#include "Misc/AutomationTest.h"
+#include "Layout/SlateRect.h"
+
+#if WITH_DEV_AUTOMATION_TESTS
+
+namespace UE::SlateWidgetAutomationTest
+{
+
+void FWidgetMockNonTemplate::ClearFunctionRecords()
+{
+	FunctionCalls.Empty();
+	DataValidations.Empty();
+}
+
+const FString& FWidgetMockNonTemplate::GetTestName() const
+{
+	return TestName;
+}
+
+TSharedPtr<FJsonObject> FWidgetMockNonTemplate::GetTestJSONObject() const
+{
+	return TestJsonObject;
+}
+
+TSharedPtr<FJsonObject> FWidgetMockNonTemplate::GetTestJSONObjectToWrite() const
+{
+	return WriteJsonObject;
+}
+
+void FWidgetMockNonTemplate::FindOrCreateTestJSONFile(const FString& InTestName)
+{
+	//@TODO Support having a separate JSON Object in the JSON test file for each mock widget. Some tests might need data from multiple mock widgets present in the test.
+	// To do this, we can assign an ID to each widget when we create it and use that ID to identify the widget. Naturally, this ID must be the same for all runs of the test.
+
+	TestName = InTestName;
+	if (FPaths::FileExists(FSlateTestHelper::GetFullPath(TestName)))
+	{
+		TestJsonObject = FSlateTestHelper::LoadJSONFromTxt(TestName);
+		WriteJsonObject = nullptr;
+		if (!TestJsonObject.IsValid())
+		{
+			DataValidations.Add(FTestIONames::NAME_FailedToLoadJSON, true);
+		}
+	}
+	else
+	{
+		DataValidations.Add(FTestIONames::NAME_TestFileNotFound, true);
+		TestJsonObject = nullptr;
+		WriteJsonObject = MakeShared<FJsonObject>();
+	}
+}
+
+bool FWidgetMockNonTemplate::IsKeyInFunctionCalls(const FName& FuncName) const
+{
+	return FunctionCalls.Contains(FuncName);
+}
+
+int32 FWidgetMockNonTemplate::GetValueInFunctionCalls(const FName& FuncName) const
+{
+	return FunctionCalls.Contains(FuncName) ? FunctionCalls[FuncName] : 0;
+}
+
+bool FWidgetMockNonTemplate::IsKeyInDataValidations(const FName& FuncName) const
+{
+	return DataValidations.Contains(FuncName);
+}
+
+bool FWidgetMockNonTemplate::GetValueInDataValidations(const FName& FuncName) const
+{
+	return DataValidations[FuncName];
+}
+
+void FWidgetMockNonTemplate::IncrementCall(const FName& FuncName) const
+{
+	int32& FunctionCallNum = FunctionCalls.FindOrAdd(FuncName);
+	FunctionCallNum += 1;
+}
+
+void FWidgetMockNonTemplate::OnPaintDataValidation(const FGeometry& AllottedGeometry, const FSlateRect& MyCullingRect, FSlateWindowElementList& OutDrawElements, int32 LayerId, const FWidgetStyle& InWidgetStyle, bool bParentEnabled) const
+{
+	// Check if we have loaded the test JSON file or we need to write it.
+	if (TestJsonObject.IsValid())
+	{
+		// Validate the parameters AllottedGeometry and InWidgetStyle.
+		if (TestJsonObject->GetStringField("AllottedGeometry") == AllottedGeometry.ToString()
+			&& TestJsonObject->GetObjectField("WidgetStyle")->GetStringField("ColorAndOpacityTint") == InWidgetStyle.GetColorAndOpacityTint().ToString()
+			&& TestJsonObject->GetObjectField("WidgetStyle")->GetStringField("ForegroundColor") == InWidgetStyle.GetForegroundColor().ToString()
+			&& TestJsonObject->GetObjectField("WidgetStyle")->GetStringField("SubduedForeground") == InWidgetStyle.GetSubduedForegroundColor().ToString())
+		{
+			// Validate the payloads in OutDrawElements.
+			for (FSlateCachedElementList* CachedElementList : OutDrawElements.GetCurrentCachedElementWithNewData())
+			{
+				if (CachedElementList != nullptr)
+				{
+					for (FSlateDrawElement& DrawElements : CachedElementList->DrawElements)
+					{
+						//@TODO Support all element types.
+						switch (DrawElements.GetElementType())
+						{
+						case EElementType::ET_ShapedText:
+						{
+							FSlateShapedTextPayload ShapedTextPayload = DrawElements.GetDataPayload<FSlateShapedTextPayload>();
+							if (TestJsonObject->GetObjectField("ShapedTextPayload")->GetStringField("Tint") == ShapedTextPayload.GetTint().ToString()
+								&& TestJsonObject->GetObjectField("ShapedTextPayload")->GetStringField("OutlineTint") == ShapedTextPayload.GetOutlineTint().ToString())
+							{
+								if (!DataValidations.Contains(FTestFunctionNames::NAME_OnPaint))
+								{
+									DataValidations.Add(FTestFunctionNames::NAME_OnPaint, true);
+								}
+							}
+							else
+							{
+								DataValidations.Add(FTestFunctionNames::NAME_OnPaint, false);
+							}
+
+						}
+						break;
+						default:
+							break;
+						}
+					}
+				}
+			}
+		}
+
+		else
+		{
+			DataValidations.Add(FTestFunctionNames::NAME_OnPaint, false);
+		}
+	}
+	else if (WriteJsonObject.IsValid())
+	{
+		FSlateTestHelper::ToJson(AllottedGeometry, WriteJsonObject);
+		FSlateTestHelper::ToJson(InWidgetStyle, WriteJsonObject);
+		for (FSlateCachedElementList* CachedElementList : OutDrawElements.GetCurrentCachedElementWithNewData())
+		{
+			if (CachedElementList != nullptr)
+			{
+				for (FSlateDrawElement& DrawElements : CachedElementList->DrawElements)
+				{
+					//@TODO Support all element types and create their own JSON conversion function.
+					switch (DrawElements.GetElementType())
+					{
+					case EElementType::ET_ShapedText:
+					{
+						FSlateShapedTextPayload ShapedTextPayload = DrawElements.GetDataPayload<FSlateShapedTextPayload>();
+						FSlateTestHelper::ToJson(ShapedTextPayload, WriteJsonObject);
+					}
+					break;
+					default:
+						break;
+					}
+				}
+			}
+		}
+	}
+	else
+	{
+		DataValidations.Add(FTestIONames::NAME_BadTestFile, true);
+	}
+}
+
+void FWidgetMockNonTemplate::ComputeDesiredSizeDataValidation(float LayoutScaleMultiplier, FVector2D RealComputedSize) const
+{
+	// Check if we have loaded the test JSON file or we need to write it.
+	if (TestJsonObject.IsValid())
+	{
+		// Validate the computed size in JSON against the size we just computed.
+		if (TestJsonObject->GetNumberField("ComputeDesiredSizeX") == RealComputedSize.X && TestJsonObject->GetNumberField("ComputeDesiredSizeY") == RealComputedSize.Y)
+		{
+			DataValidations.Add(FTestFunctionNames::NAME_ComputeDesiredSize, true);
+
+		}
+		else
+		{
+			DataValidations.Add(FTestFunctionNames::NAME_ComputeDesiredSize, false);
+		}
+	}
+	else if (WriteJsonObject.IsValid())
+	{
+		WriteJsonObject->SetNumberField("ComputeDesiredSizeX", RealComputedSize.X);
+		WriteJsonObject->SetNumberField("ComputeDesiredSizeY", RealComputedSize.Y);
+	}
+	// If we haven't loaded the JSON file and haven't initialized a JSON object to write to, this is an error and we can't continue the test.
+	else
+	{
+		DataValidations.Add(FTestIONames::NAME_BadTestFile, true);
+	}
+}
+
+}
+#endif
\ No newline at end of file
diff --git a/Engine/Plugins/Tests/WidgetAutomationTests/Source/WidgetAutomationTests/Private/WidgetMockNonTemplate.h b/Engine/Plugins/Tests/WidgetAutomationTests/Source/WidgetAutomationTests/Private/WidgetMockNonTemplate.h
new file mode 100644
--- /dev/null
+++ b/Engine/Plugins/Tests/WidgetAutomationTests/Source/WidgetAutomationTests/Private/WidgetMockNonTemplate.h
@@ -0,0 +1,64 @@
+// Copyright Epic Games, Inc. All Rights Reserved.
+#pragma once
+
+#include "CoreMinimal.h"
+
+#if WITH_DEV_AUTOMATION_TESTS
+class FJsonObject;
+struct FGeometry;
+class FSlateRect;
+class FSlateWindowElementList;
+class FWidgetStyle;
+
+namespace UE::SlateWidgetAutomationTest
+{
+
+// Do NOT create instances of this class. It is used for refactoring the template class SWidgetMock. Use SWidgetMock for your tests.
+class FWidgetMockNonTemplate
+{
+public:
+	void ClearFunctionRecords();
+
+	const FString& GetTestName() const;
+	TSharedPtr<FJsonObject> GetTestJSONObject() const;
+	TSharedPtr<FJsonObject> GetTestJSONObjectToWrite() const;
+	void FindOrCreateTestJSONFile(const FString& InTestName);
+
+	bool IsKeyInFunctionCalls(const FName& FuncName) const;
+
+	// Return the number of times we called the function with the name FuncName.
+	int32 GetValueInFunctionCalls(const FName& FuncName) const;
+
+	bool IsKeyInDataValidations(const FName& FuncName) const;
+	bool GetValueInDataValidations(const FName& FuncName) const;
+
+	// Increase (by 1) the number of times we called the function with the name FuncName.
+	void IncrementCall(const FName& FuncName) const;
+
+	void OnPaintDataValidation(const FGeometry& AllottedGeometry, const FSlateRect& MyCullingRect, FSlateWindowElementList& OutDrawElements, int32 LayerId, const FWidgetStyle& InWidgetStyle, bool bParentEnabled) const;
+
+	void ComputeDesiredSizeDataValidation(float LayoutScaleMultiplier, FVector2D RealComputedSize) const;
+
+private:
+	// Map to store function calls and the number of times they were called.
+	// Ex. (OnComputeSize, 3) means OnComputeSize was called 3 times.
+	mutable TMap<FName, int32> FunctionCalls;
+
+	// Map to store the result of the verification of parameters and outputs of a function.
+	// Ex. (OnComputeSize, true) means the parameters and output of the function OnComputeSize were correct according to the test JSON baseline.
+	mutable TMap<FName, bool> DataValidations;
+
+	// This will be incorporated in the name of the JSON baseline file.
+	// Make sure it is unique and that it is set in all important widgets. 
+	FString TestName;
+
+	// The JSON object loaded from the test JSON file on disk.
+	TSharedPtr<FJsonObject> TestJsonObject;
+
+	// JSON object that collects and holds all data that needs to be written to the JSON file.
+	// If we initially find a JSON file for the test, this will remain null.
+	TSharedPtr<FJsonObject> WriteJsonObject;
+};
+
+}
+#endif
\ No newline at end of file
diff --git a/Engine/Plugins/Tests/WidgetAutomationTests/Source/WidgetAutomationTests/Public/IWidgetAutomationTestsModule.h b/Engine/Plugins/Tests/WidgetAutomationTests/Source/WidgetAutomationTests/Public/IWidgetAutomationTestsModule.h
new file mode 100644
--- /dev/null
+++ b/Engine/Plugins/Tests/WidgetAutomationTests/Source/WidgetAutomationTests/Public/IWidgetAutomationTestsModule.h
@@ -0,0 +1,40 @@
+// Copyright Epic Games, Inc. All Rights Reserved.
+
+#pragma once
+
+#include "CoreMinimal.h"
+#include "Modules/ModuleInterface.h"
+#include "Modules/ModuleManager.h"
+
+
+/**
+ * The public interface to this module.  In most cases, this interface is only public to sibling modules
+ * within this plugin.
+ */
+class IWidgetAutomationTests : public IModuleInterface
+{
+
+public:
+
+	/**
+	 * Singleton-like access to this module's interface.  This is just for convenience!
+	 * Beware of calling this during the shutdown phase, though.  Your module might have been unloaded already.
+	 *
+	 * @return Returns singleton instance, loading the module on demand if needed
+	 */
+	static inline IWidgetAutomationTests& Get()
+	{
+		return FModuleManager::LoadModuleChecked<IWidgetAutomationTests>("WidgetAutomationTests");
+	}
+
+	/**
+	 * Checks to see if this module is loaded and ready.  It is only valid to call Get() if IsAvailable() returns true.
+	 *
+	 * @return True if the module is loaded and ready to use
+	 */
+	static inline bool IsAvailable()
+	{
+		return FModuleManager::Get().IsModuleLoaded("WidgetAutomationTests");
+	}
+};
+
diff --git a/Engine/Plugins/Tests/WidgetAutomationTests/Source/WidgetAutomationTests/Public/SWidgetMock.h b/Engine/Plugins/Tests/WidgetAutomationTests/Source/WidgetAutomationTests/Public/SWidgetMock.h
new file mode 100644
--- /dev/null
+++ b/Engine/Plugins/Tests/WidgetAutomationTests/Source/WidgetAutomationTests/Public/SWidgetMock.h
@@ -0,0 +1,60 @@
+// Copyright Epic Games, Inc. All Rights Reserved.
+#pragma once
+
+#include "CoreMinimal.h"
+#include "TestDataAndFunctionNames.h"
+#include "WidgetMockNonTemplate.h"
+#include "Widgets/SWidget.h"
+
+#if WITH_DEV_AUTOMATION_TESTS
+
+namespace UE::SlateWidgetAutomationTest
+{
+// All widgets in a test should inherit from this class.
+template<class T>
+class SWidgetMock : public T, public FWidgetMockNonTemplate
+{
+	static_assert(TIsDerivedFrom<T, SWidget>::IsDerived, "The mock widget must be a subclass of SWidget.");
+public:
+
+	virtual FReply OnFocusReceived(const FGeometry& MyGeometry, const FFocusEvent& InFocusEvent) override
+	{
+		T::OnFocusReceived(MyGeometry, InFocusEvent);
+		IncrementCall(FTestFunctionNames::NAME_OnFocusReceived);
+		return FReply::Unhandled();
+	}
+
+	virtual void OnFocusLost(const FFocusEvent& InFocusEvent) override
+	{
+		T::OnFocusLost(InFocusEvent);
+		IncrementCall(FTestFunctionNames::NAME_OnFocusLost);
+	}
+
+	virtual void OnFocusChanging(const FWeakWidgetPath& PreviousFocusPath, const FWidgetPath& NewWidgetPath, const FFocusEvent& InFocusEvent) override
+	{
+		T::OnFocusChanging(PreviousFocusPath, NewWidgetPath, InFocusEvent);
+		IncrementCall(FTestFunctionNames::NAME_OnFocusChanging);
+	}
+
+protected:
+	virtual FVector2D ComputeDesiredSize(float LayoutScaleMultiplier) const override
+	{
+		IncrementCall(FTestFunctionNames::NAME_ComputeDesiredSize);
+		FVector2D RealComputedSize = T::ComputeDesiredSize(LayoutScaleMultiplier);
+
+		ComputeDesiredSizeDataValidation(LayoutScaleMultiplier, RealComputedSize);
+		return RealComputedSize;
+	}
+
+private:
+	virtual int32 OnPaint(const FPaintArgs& Args, const FGeometry& AllottedGeometry, const FSlateRect& MyCullingRect, FSlateWindowElementList& OutDrawElements, int32 LayerId, const FWidgetStyle& InWidgetStyle, bool bParentEnabled) const override
+	{
+		IncrementCall(FTestFunctionNames::NAME_OnPaint);
+		int32 RealLayer = T::OnPaint(Args, AllottedGeometry, MyCullingRect, OutDrawElements, LayerId, InWidgetStyle, bParentEnabled);
+
+		this->OnPaintDataValidation(AllottedGeometry, MyCullingRect, OutDrawElements, LayerId, InWidgetStyle, bParentEnabled);
+		return RealLayer;
+	}
+};
+}
+#endif
\ No newline at end of file
diff --git a/Engine/Plugins/Tests/WidgetAutomationTests/Source/WidgetAutomationTests/Public/SlateTestBase.h b/Engine/Plugins/Tests/WidgetAutomationTests/Source/WidgetAutomationTests/Public/SlateTestBase.h
new file mode 100644
--- /dev/null
+++ b/Engine/Plugins/Tests/WidgetAutomationTests/Source/WidgetAutomationTests/Public/SlateTestBase.h
@@ -0,0 +1,143 @@
+// Copyright Epic Games, Inc. All Rights Reserved.
+#pragma once
+
+#include "CoreMinimal.h"
+#include "Misc/AutomationTest.h"
+#include "Misc/GeneratedTypeName.h"
+#include "SlateTestHelper.h"
+#include "TestDataAndFunctionNames.h"
+#include "Widgets/InvalidateWidgetReason.h"
+#include "Widgets/SWindow.h"
+
+#if WITH_DEV_AUTOMATION_TESTS
+
+class SWidget;
+
+namespace UE::SlateWidgetAutomationTest
+{
+template<class T>
+class SWidgetMock;
+
+// Base interface that all slate widget tests should inherit from
+template <class T>
+class FSlateTestBase
+{
+	static_assert(TIsDerivedFrom<T, SWidget>::IsDerived, "The template class of the test must be a subclass of SWidget.");
+
+public:
+
+	FSlateTestBase() : FSlateTestBase(TEXT(""))
+	{}
+
+	FSlateTestBase(FString InTestName) : 
+		MainWindow(SNew(SWindow)
+			.ClientSize(FVector2D(1000.f, 700.f))
+			.SupportsMinimize(false)),
+		SubjectWidget(CreateWidget()),
+		TestName(MoveTemp(InTestName))
+	{
+		FSlateApplication::Get().AddWindow(MainWindow);
+	}
+
+	virtual void Setup(FAutomationTestBase* TestObj) = 0;
+
+	virtual void Run(FAutomationTestBase* TestObj) = 0;
+
+	virtual void Validate(FAutomationTestBase* TestObj) = 0;
+
+	virtual void TearDown(FAutomationTestBase* TestObj) = 0;
+
+	TSharedRef<SWindow> GetMainWindow() const
+	{
+		return MainWindow;
+	}
+
+	TSharedRef<SWidgetMock<T>> GetSubjectWidget() const
+	{
+		return SubjectWidget;
+	}
+
+	TSharedRef< SWidgetMock<T>> CreateWidget() const
+	{
+		return SNew(SWidgetMock<T>);
+	}
+
+	const FString& GetTestName() const
+	{
+		return TestName;
+	}
+
+	// This function consists of generic verifications that we want to do for every property change in a SWidget.
+	template< typename PropertyType>
+	void ValidatePropertyChange(FAutomationTestBase* TestObj, typename TMemFunPtrType<true, T, PropertyType()>::Type Getter, const FString& PropertyName, PropertyType TestValue, EInvalidateWidgetReason InvalidateReason)
+	{
+		// 1. Verify that the property getter returns the value we set. 
+		TestObj->AddErrorIfFalse((SubjectWidget.Get().*Getter)() == TestValue, FString::Printf(TEXT("In %s: The property %s was not set."), GetGeneratedTypeName<T>(), *PropertyName));
+
+		// 2. Ensure global invalidation is on. Otherwise we will skip checking if the widget has a valid handle. 
+		if (!MainWindow->Advanced_IsInvalidationRoot())
+		{
+			TestObj->AddError(FString::Printf(TEXT("In %s: Slate.EnableGlobalInvalidation is disabled. Please enable it and run the test again."), GetGeneratedTypeName<T>()));
+		}
+		else
+		{
+			TestObj->AddErrorIfFalse(SubjectWidget->GetProxyHandle().IsValid(SubjectWidget.Get()), FString::Printf(TEXT("In %s: The widget doesn't have a valid Invalidation Root. Make sure the widget is added to a window in the previous frame, and the window has fast update enabled."), GetGeneratedTypeName<T>()));
+
+			// 3. In the next frame, verify if requested by the invalidation reason:
+			// - ComputeDesiredSize was called on the widget whose property has changed
+			// - ComputeDesiredSize returns the correct size.
+			// - OnPaint was called on the widget whose property has changed
+			// - OnPaint has the correct values for Allotted Geometry and WidgetStyle
+			TestObj->AddCommand(new FDelayedFunctionLatentCommand([this, WidgetRef = GetSubjectWidget(), TestObj, InvalidateReason, PropertyName, MyWindow = MainWindow]()
+				{
+					// We have to do this JSON check here because we attempt to read/write the JSON file in this frame for the first time.
+					if (WidgetRef->IsKeyInDataValidations(FTestIONames::NAME_BadTestFile))
+					{
+						FString FileLocation = FSlateTestHelper::GetFullPath(TestName);
+						TestObj->AddError(FString::Printf(TEXT("In property %s in %s. The ground truth JSON file for this test was not found and could not be created. Did you forget to call FindOrCreateTestJSONFile in the setup of your test?"), *PropertyName, GetGeneratedTypeName<T>(), *FileLocation));
+					}
+					else
+					{
+						if (EnumHasAnyFlags(InvalidateReason, EInvalidateWidgetReason::Layout))
+						{
+							TestObj->AddErrorIfFalse(WidgetRef->GetValueInFunctionCalls(FTestFunctionNames::NAME_ComputeDesiredSize) == 1, FString::Printf(TEXT("In %s: ComputeDesiredSize was not called after the property %s changed."), GetGeneratedTypeName<T>(), *PropertyName));
+							TestObj->AddErrorIfFalse(WidgetRef->IsKeyInDataValidations(FTestFunctionNames::NAME_ComputeDesiredSize) && WidgetRef->GetValueInDataValidations(FTestFunctionNames::NAME_ComputeDesiredSize), FString::Printf(TEXT("In %s: The computed desired size is wrong or does not exist after property %s changed."), GetGeneratedTypeName<T>(), *PropertyName));
+						}
+
+						if (EnumHasAnyFlags(InvalidateReason, EInvalidateWidgetReason::Paint))
+						{
+							TestObj->AddErrorIfFalse(WidgetRef->GetValueInFunctionCalls(FTestFunctionNames::NAME_OnPaint) == 1, FString::Printf(TEXT("In %s: OnPaint was not called after the property %s changed in widget. Make sure the widget is being invalidated."), GetGeneratedTypeName<T>(), *PropertyName));
+							TestObj->AddErrorIfFalse(WidgetRef->IsKeyInDataValidations(FTestFunctionNames::NAME_OnPaint) && WidgetRef->GetValueInDataValidations(FTestFunctionNames::NAME_OnPaint), FString::Printf(TEXT("In %s: The list of elements to draw is wrong or does not exist."), GetGeneratedTypeName<T>(), *PropertyName));
+						}
+					}
+
+				}, 0.1f));
+		}
+	};
+
+	virtual ~FSlateTestBase() = default;
+
+	bool JSONValidation(FAutomationTestBase* TestObj, const FString& PropertyName)
+	{
+		if (SubjectWidget->IsKeyInDataValidations(FTestIONames::NAME_TestFileNotFound))
+		{
+			FString FileLocation = FSlateTestHelper::GetFullPath(TestName);
+			TestObj->AddError(FString::Printf(TEXT("Skipping the test for the property %s in %s. The ground truth JSON file for this test was not found looking under path: %s.\n This is expected only if this test is being run for the first time after creation or modification. If this is the case, please run the test again."), *PropertyName, GetGeneratedTypeName<T>(), *FileLocation));
+			return false;
+		}
+		if (SubjectWidget->IsKeyInDataValidations(FTestIONames::NAME_FailedToLoadJSON))
+		{
+			TestObj->AddError(FString::Printf(TEXT("Skipping the test for the property %s in %s. The ground truth JSON file for this test was found but could not be loaded."), *PropertyName, GetGeneratedTypeName<T>()));
+			return false;
+		}
+		return true;
+	}
+
+private:
+	TSharedRef<SWindow> MainWindow;
+	TSharedRef<SWidgetMock<T>> SubjectWidget; 
+	FString TestName;
+};
+}
+
+#endif
\ No newline at end of file
diff --git a/Engine/Plugins/Tests/WidgetAutomationTests/Source/WidgetAutomationTests/Public/SlateTestHelper.h b/Engine/Plugins/Tests/WidgetAutomationTests/Source/WidgetAutomationTests/Public/SlateTestHelper.h
new file mode 100644
--- /dev/null
+++ b/Engine/Plugins/Tests/WidgetAutomationTests/Source/WidgetAutomationTests/Public/SlateTestHelper.h
@@ -0,0 +1,34 @@
+// Copyright Epic Games, Inc. All Rights Reserved.
+#pragma once
+
+#include "CoreMinimal.h"
+#include "Dom/JsonObject.h"
+#include "Styling/WidgetStyle.h"
+#include "Layout/Geometry.h"
+
+#if WITH_DEV_AUTOMATION_TESTS
+
+class JsonObject;
+struct FSlateShapedTextPayload;
+
+namespace UE::SlateWidgetAutomationTest
+{
+// A collection of helper functions used for slate tests.
+class FSlateTestHelper
+{
+public:	
+	// Write ObjToWrite to the JSON file with name TestName.
+	static bool WriteJSONToTxt(TSharedRef<FJsonObject> ObjToWrite, const FString& TestName);
+
+	// Load and return the JSON object in JSON file TestName.
+	static TSharedPtr<FJsonObject> LoadJSONFromTxt(const FString& TestName);
+
+	static FString GetFullPath(FStringView TestName);
+	
+	// JSON conversion functions.
+	static void ToJson(const FGeometry& MyGeometry, TSharedPtr<FJsonObject> MyJsonObject);
+	static void ToJson(const FWidgetStyle& InWidgetStyle, TSharedPtr<FJsonObject> MyJsonObject);
+	static void ToJson(const FSlateShapedTextPayload& ShapedTextPayload, TSharedPtr<FJsonObject> MyJsonObject);
+};
+}
+#endif
\ No newline at end of file
diff --git a/Engine/Plugins/Tests/WidgetAutomationTests/Source/WidgetAutomationTests/Public/TestDataAndFunctionNames.h b/Engine/Plugins/Tests/WidgetAutomationTests/Source/WidgetAutomationTests/Public/TestDataAndFunctionNames.h
new file mode 100644
--- /dev/null
+++ b/Engine/Plugins/Tests/WidgetAutomationTests/Source/WidgetAutomationTests/Public/TestDataAndFunctionNames.h
@@ -0,0 +1,25 @@
+// Copyright Epic Games, Inc. All Rights Reserved.
+#pragma once
+#include "UObject/NameTypes.h"
+
+#if WITH_DEV_AUTOMATION_TESTS
+
+namespace UE::SlateWidgetAutomationTest
+{
+	namespace FTestIONames
+	{
+		 extern const FLazyName NAME_FailedToLoadJSON;
+		 extern const FLazyName NAME_TestFileNotFound;
+		 extern const FLazyName NAME_BadTestFile;
+	};
+
+	namespace FTestFunctionNames
+	{
+		extern const FLazyName NAME_OnPaint;
+		extern const FLazyName NAME_ComputeDesiredSize;
+		extern const FLazyName NAME_OnFocusChanging;
+		extern const FLazyName NAME_OnFocusLost;
+		extern const FLazyName NAME_OnFocusReceived;
+	};
+}
+#endif
\ No newline at end of file
diff --git a/Engine/Plugins/Tests/WidgetAutomationTests/Source/WidgetAutomationTests/WidgetAutomationTests.Build.cs b/Engine/Plugins/Tests/WidgetAutomationTests/Source/WidgetAutomationTests/WidgetAutomationTests.Build.cs
new file mode 100644
--- /dev/null
+++ b/Engine/Plugins/Tests/WidgetAutomationTests/Source/WidgetAutomationTests/WidgetAutomationTests.Build.cs
@@ -0,0 +1,57 @@
+// Copyright Epic Games, Inc. All Rights Reserved.
+
+using UnrealBuildTool;
+
+public class WidgetAutomationTests : ModuleRules
+{
+	public WidgetAutomationTests(ReadOnlyTargetRules Target) : base(Target)
+	{
+		
+		PublicIncludePaths.AddRange(
+			new string[] {
+				// ... add public include paths required here ...
+			}
+			);
+				
+		
+		PrivateIncludePaths.AddRange(
+			new string[] {
+				// ... add other private include paths required here ...
+			}
+			);
+			
+		
+		PublicDependencyModuleNames.AddRange(
+			new string[]
+			{   
+				"Core",
+				"CoreUObject",
+				"EditorTests",
+				"Engine",
+				"Slate",
+				"SlateCore",
+				"Json",
+				"UnrealEd",
+				"UMG",
+				"UMGEditor"
+				// ... add other public dependencies that you statically link with here ...
+			}
+			);
+			
+		
+		PrivateDependencyModuleNames.AddRange(
+			new string[]
+			{
+				// ... add private dependencies that you statically link with here ...	
+			}
+			);
+		
+		
+		DynamicallyLoadedModuleNames.AddRange(
+			new string[]
+			{
+				// ... add any modules that your module loads dynamically here ...
+			}
+			);
+	}
+}
diff --git a/Engine/Plugins/Tests/WidgetAutomationTests/WidgetAutomationTests.uplugin b/Engine/Plugins/Tests/WidgetAutomationTests/WidgetAutomationTests.uplugin
new file mode 100644
--- /dev/null
+++ b/Engine/Plugins/Tests/WidgetAutomationTests/WidgetAutomationTests.uplugin
@@ -0,0 +1,31 @@
+{
+	"FileVersion": 3,
+	"Version": 1,
+	"VersionName": "1.0",
+	"FriendlyName": "WidgetAutomationTests",
+	"Description": "",
+	"Category": "Testing",
+	"CreatedBy": "Epic Games, Inc",
+	"CreatedByURL": "https://epicgames.com",
+	"DocsURL": "",
+	"MarketplaceURL": "",
+	"SupportURL": "",
+	"EnabledByDefault": false,
+	"CanContainContent": true,
+	"IsBetaVersion": false,
+	"IsExperimentalVersion": true,
+	"Installed": false,
+	"Modules": [
+		{
+			"Name": "WidgetAutomationTests",
+			"Type": "Editor",
+			"LoadingPhase": "Default"
+		}
+	],
+	"Plugins": [
+		{
+			"Name": "EditorTests",
+			"Enabled": true
+		}
+	]
+}
\ No newline at end of file
diff --git a/Engine/Source/Runtime/Slate/Private/Widgets/Text/STextBlock.cpp b/Engine/Source/Runtime/Slate/Private/Widgets/Text/STextBlock.cpp
--- a/Engine/Source/Runtime/Slate/Private/Widgets/Text/STextBlock.cpp
+++ b/Engine/Source/Runtime/Slate/Private/Widgets/Text/STextBlock.cpp
@@ -168,6 +168,16 @@ ETextTransformPolicy STextBlock::GetTransformPolicyImpl() const
 	return bIsAttributeTransformPolicySet ? TransformPolicy.Get() : TextStyle.TransformPolicy;
 }
 
+FMargin STextBlock::GetMargin() const
+{
+	return Margin.Get();
+}
+
+float STextBlock::GetMinDesiredWidth() const
+{
+	return MinDesiredWidth.Get();
+}
+
 void STextBlock::InvalidateText(EInvalidateWidgetReason InvalidateReason)
 {
 	if (bSimpleTextMode && EnumHasAnyFlags(InvalidateReason, EInvalidateWidgetReason::Layout))
diff --git a/Engine/Source/Runtime/Slate/Public/Widgets/Input/SButton.h b/Engine/Source/Runtime/Slate/Public/Widgets/Input/SButton.h
--- a/Engine/Source/Runtime/Slate/Public/Widgets/Input/SButton.h
+++ b/Engine/Source/Runtime/Slate/Public/Widgets/Input/SButton.h
@@ -208,7 +208,7 @@ public:
 #if WITH_ACCESSIBILITY
 	virtual TSharedRef<FSlateAccessibleWidget> CreateAccessibleWidget() override;
 #endif
-private:
+protected:
 	virtual FVector2D ComputeDesiredSize(float) const override;
 	//~ SWidget
 
diff --git a/Engine/Source/Runtime/Slate/Public/Widgets/Text/STextBlock.h b/Engine/Source/Runtime/Slate/Public/Widgets/Text/STextBlock.h
--- a/Engine/Source/Runtime/Slate/Public/Widgets/Text/STextBlock.h
+++ b/Engine/Source/Runtime/Slate/Public/Widgets/Text/STextBlock.h
@@ -255,10 +255,7 @@ public:
 #endif
 	// End of SWidget interface
 
-private:
-	/** Get the computed text style to use with the text marshaller */
-	FTextBlockStyle GetComputedTextStyle() const;
-
+public:
 	/** Gets the current foreground color */
 	FSlateColor GetColorAndOpacity() const;
 
@@ -283,11 +280,18 @@ private:
 	/** Gets the current highlight shape */
 	const FSlateBrush* GetHighlightShape() const;
 
-	/** Call to invalidate this text block */
-	void InvalidateText(EInvalidateWidgetReason InvalidateReason);
+	/** Gets the current margin */
+	FMargin GetMargin() const;
+
+	/** Gets the current min desired width */
+	float GetMinDesiredWidth() const;
 
 private:
+	/** Get the computed text style to use with the text marshaller */
+	FTextBlockStyle GetComputedTextStyle() const;
 
+	/** Call to invalidate this text block */
+	void InvalidateText(EInvalidateWidgetReason InvalidateReason);
 
 private:
 	/** The text displayed in this text block */
diff --git a/Engine/Source/Runtime/SlateCore/Public/Rendering/DrawElements.h b/Engine/Source/Runtime/SlateCore/Public/Rendering/DrawElements.h
--- a/Engine/Source/Runtime/SlateCore/Public/Rendering/DrawElements.h
+++ b/Engine/Source/Runtime/SlateCore/Public/Rendering/DrawElements.h
@@ -460,6 +460,7 @@ struct FSlateCachedElementData
 
 	void AddReferencedObjects(FReferenceCollector& Collector);
 
+	TArrayView<FSlateCachedElementList* const> GetListsWithNewData() const { return MakeArrayView(ListsWithNewData.GetData(), ListsWithNewData.Num()); }
 private:
 	/** Removes a cache node completely from the cache */
 	void RemoveList(FSlateCachedElementsHandle& CacheHandle);
@@ -669,6 +670,17 @@ public:
 	SLATECORE_API void SetIsInGameLayer(bool bInGameLayer);
 	SLATECORE_API bool GetIsInGameLayer();
 
+	TArrayView<FSlateCachedElementList* const> GetCurrentCachedElementWithNewData() const
+	{ 
+		if (GetCurrentCachedElementData() != nullptr)
+		{
+			return CachedElementDataList[CachedElementDataListStack.Top()]->GetListsWithNewData();
+		}
+		else
+		{
+			return TArrayView<FSlateCachedElementList* const>();
+		}
+	}
 private:
 	FSlateDrawElement& AddCachedElement();
 
