Add copyright notice to a .cpp file in WidgetAutomationTest plugin.

diff --git a/Engine/Plugins/Tests/WidgetAutomationTests/Source/WidgetAutomationTests/Private/TestDataAndFunctionNames.cpp b/Engine/Plugins/Tests/WidgetAutomationTests/Source/WidgetAutomationTests/Private/TestDataAndFunctionNames.cpp
--- a/Engine/Plugins/Tests/WidgetAutomationTests/Source/WidgetAutomationTests/Private/TestDataAndFunctionNames.cpp
+++ b/Engine/Plugins/Tests/WidgetAutomationTests/Source/WidgetAutomationTests/Private/TestDataAndFunctionNames.cpp
@@ -1,3 +1,4 @@
+// Copyright Epic Games, Inc. All Rights Reserved.
 #include "TestDataAndFunctionNames.h"
 
 #if WITH_DEV_AUTOMATION_TESTS
