Widget automation test framework: fix non unity error in Mac.

diff --git a/Engine/Plugins/Tests/WidgetAutomationTests/Source/WidgetAutomationTests/Public/SlateTestBase.h b/Engine/Plugins/Tests/WidgetAutomationTests/Source/WidgetAutomationTests/Public/SlateTestBase.h
--- a/Engine/Plugins/Tests/WidgetAutomationTests/Source/WidgetAutomationTests/Public/SlateTestBase.h
+++ b/Engine/Plugins/Tests/WidgetAutomationTests/Source/WidgetAutomationTests/Public/SlateTestBase.h
@@ -2,6 +2,7 @@
 #pragma once
 
 #include "CoreMinimal.h"
+#include "Framework/Application/SlateApplication.h"
 #include "Misc/AutomationTest.h"
 #include "Misc/GeneratedTypeName.h"
 #include "SlateTestHelper.h"
