Update WidgetMock to use new element payload list type

diff --git a/Engine/Plugins/Tests/WidgetAutomationTests/Source/WidgetAutomationTests/Private/WidgetMockNonTemplate.cpp b/Engine/Plugins/Tests/WidgetAutomationTests/Source/WidgetAutomationTests/Private/WidgetMockNonTemplate.cpp
--- a/Engine/Plugins/Tests/WidgetAutomationTests/Source/WidgetAutomationTests/Private/WidgetMockNonTemplate.cpp
+++ b/Engine/Plugins/Tests/WidgetAutomationTests/Source/WidgetAutomationTests/Private/WidgetMockNonTemplate.cpp
@@ -102,27 +102,32 @@ void FWidgetMockNonTemplate::OnPaintDataValidation(const FGeometry& AllottedGeom
 			{
 				if (CachedElementList != nullptr)
 				{
-					for (FSlateDrawElement& DrawElements : CachedElementList->DrawElements)
+					for (TPair<EElementType, FSlateDrawElementContainer>& DrawElementMap : CachedElementList->DrawElements)
 					{
+						EElementType ElementType = DrawElementMap.Key;
+						const FSlateDrawElementContainer& Container = DrawElementMap.Value;
+
 						//@TODO Support all element types.
-						switch (DrawElements.GetElementType())
+						switch (ElementType)
 						{
 						case EElementType::ET_ShapedText:
 						{
-							FSlateShapedTextPayload ShapedTextPayload = DrawElements.GetDataPayload<FSlateShapedTextPayload>();
-							if (TestJsonObject->GetObjectField("ShapedTextPayload")->GetStringField("Tint") == ShapedTextPayload.GetTint().ToString()
-								&& TestJsonObject->GetObjectField("ShapedTextPayload")->GetStringField("OutlineTint") == ShapedTextPayload.GetOutlineTint().ToString())
+							for (FSlateDrawElement& DrawElements : Container.Elements)
 							{
-								if (!DataValidations.Contains(FTestFunctionNames::NAME_OnPaint))
+								FSlateShapedTextPayload ShapedTextPayload = DrawElements.GetDataPayload<FSlateShapedTextPayload>();
+								if (TestJsonObject->GetObjectField("ShapedTextPayload")->GetStringField("Tint") == ShapedTextPayload.GetTint().ToString()
+									&& TestJsonObject->GetObjectField("ShapedTextPayload")->GetStringField("OutlineTint") == ShapedTextPayload.GetOutlineTint().ToString())
 								{
-									DataValidations.Add(FTestFunctionNames::NAME_OnPaint, true);
+									if (!DataValidations.Contains(FTestFunctionNames::NAME_OnPaint))
+									{
+										DataValidations.Add(FTestFunctionNames::NAME_OnPaint, true);
+									}
+								}
+								else
+								{
+									DataValidations.Add(FTestFunctionNames::NAME_OnPaint, false);
 								}
 							}
-							else
-							{
-								DataValidations.Add(FTestFunctionNames::NAME_OnPaint, false);
-							}
-
 						}
 						break;
 						default:
@@ -146,15 +151,21 @@ void FWidgetMockNonTemplate::OnPaintDataValidation(const FGeometry& AllottedGeom
 		{
 			if (CachedElementList != nullptr)
 			{
-				for (FSlateDrawElement& DrawElements : CachedElementList->DrawElements)
+				for (TPair<EElementType, FSlateDrawElementContainer>& DrawElementMap : CachedElementList->DrawElements)
 				{
+					EElementType ElementType = DrawElementMap.Key;
+					const FSlateDrawElementContainer& Container = DrawElementMap.Value;
+
 					//@TODO Support all element types and create their own JSON conversion function.
-					switch (DrawElements.GetElementType())
+					switch (ElementType)
 					{
 					case EElementType::ET_ShapedText:
 					{
-						FSlateShapedTextPayload ShapedTextPayload = DrawElements.GetDataPayload<FSlateShapedTextPayload>();
-						FSlateTestHelper::ToJson(ShapedTextPayload, WriteJsonObject);
+						for (const FSlateDrawElement& DrawElements : Container.Elements)
+						{
+							FSlateShapedTextPayload ShapedTextPayload = DrawElements.GetDataPayload<FSlateShapedTextPayload>();
+							FSlateTestHelper::ToJson(ShapedTextPayload, WriteJsonObject);
+						}
 					}
 					break;
 					default:
