Add missing const WidgetMock

diff --git a/Engine/Plugins/Tests/WidgetAutomationTests/Source/WidgetAutomationTests/Private/WidgetMockNonTemplate.cpp b/Engine/Plugins/Tests/WidgetAutomationTests/Source/WidgetAutomationTests/Private/WidgetMockNonTemplate.cpp
--- a/Engine/Plugins/Tests/WidgetAutomationTests/Source/WidgetAutomationTests/Private/WidgetMockNonTemplate.cpp
+++ b/Engine/Plugins/Tests/WidgetAutomationTests/Source/WidgetAutomationTests/Private/WidgetMockNonTemplate.cpp
@@ -112,7 +112,7 @@ void FWidgetMockNonTemplate::OnPaintDataValidation(const FGeometry& AllottedGeom
 						{
 						case EElementType::ET_ShapedText:
 						{
-							for (FSlateDrawElement& DrawElements : Container.Elements)
+							for (const FSlateDrawElement& DrawElements : Container.Elements)
 							{
 								FSlateShapedTextPayload ShapedTextPayload = DrawElements.GetDataPayload<FSlateShapedTextPayload>();
 								if (TestJsonObject->GetObjectField("ShapedTextPayload")->GetStringField("Tint") == ShapedTextPayload.GetTint().ToString()
