Fix Mock Template CIS

diff --git a/Engine/Plugins/Tests/WidgetAutomationTests/Source/WidgetAutomationTests/Private/WidgetMockNonTemplate.cpp b/Engine/Plugins/Tests/WidgetAutomationTests/Source/WidgetAutomationTests/Private/WidgetMockNonTemplate.cpp
--- a/Engine/Plugins/Tests/WidgetAutomationTests/Source/WidgetAutomationTests/Private/WidgetMockNonTemplate.cpp
+++ b/Engine/Plugins/Tests/WidgetAutomationTests/Source/WidgetAutomationTests/Private/WidgetMockNonTemplate.cpp
@@ -102,36 +102,21 @@ void FWidgetMockNonTemplate::OnPaintDataValidation(const FGeometry& AllottedGeom
 			{
 				if (CachedElementList != nullptr)
 				{
-					for (TPair<EElementType, FSlateDrawElementContainer>& DrawElementMap : CachedElementList->DrawElements)
+					//@TODO Support all element types and create their own JSON conversion function.
+					for (const FSlateDrawElement& DrawElements : CachedElementList->DrawElements.Get<(uint8)EElementType::ET_ShapedText>().Elements)
 					{
-						EElementType ElementType = DrawElementMap.Key;
-						const FSlateDrawElementContainer& Container = DrawElementMap.Value;
-
-						//@TODO Support all element types.
-						switch (ElementType)
-						{
-						case EElementType::ET_ShapedText:
+						FSlateShapedTextPayload ShapedTextPayload = DrawElements.GetDataPayload<FSlateShapedTextPayload>();
+						if (TestJsonObject->GetObjectField("ShapedTextPayload")->GetStringField("Tint") == ShapedTextPayload.GetTint().ToString()
+							&& TestJsonObject->GetObjectField("ShapedTextPayload")->GetStringField("OutlineTint") == ShapedTextPayload.GetOutlineTint().ToString())
 						{
-							for (const FSlateDrawElement& DrawElements : Container.Elements)
+							if (!DataValidations.Contains(FTestFunctionNames::NAME_OnPaint))
 							{
-								FSlateShapedTextPayload ShapedTextPayload = DrawElements.GetDataPayload<FSlateShapedTextPayload>();
-								if (TestJsonObject->GetObjectField("ShapedTextPayload")->GetStringField("Tint") == ShapedTextPayload.GetTint().ToString()
-									&& TestJsonObject->GetObjectField("ShapedTextPayload")->GetStringField("OutlineTint") == ShapedTextPayload.GetOutlineTint().ToString())
-								{
-									if (!DataValidations.Contains(FTestFunctionNames::NAME_OnPaint))
-									{
-										DataValidations.Add(FTestFunctionNames::NAME_OnPaint, true);
-									}
-								}
-								else
-								{
-									DataValidations.Add(FTestFunctionNames::NAME_OnPaint, false);
-								}
+								DataValidations.Add(FTestFunctionNames::NAME_OnPaint, true);
 							}
 						}
-						break;
-						default:
-							break;
+						else
+						{
+							DataValidations.Add(FTestFunctionNames::NAME_OnPaint, false);
 						}
 					}
 				}
@@ -151,26 +136,11 @@ void FWidgetMockNonTemplate::OnPaintDataValidation(const FGeometry& AllottedGeom
 		{
 			if (CachedElementList != nullptr)
 			{
-				for (TPair<EElementType, FSlateDrawElementContainer>& DrawElementMap : CachedElementList->DrawElements)
+				//@TODO Support all element types and create their own JSON conversion function.
+				for (const FSlateDrawElement& DrawElements : CachedElementList->DrawElements.Get<(uint8)EElementType::ET_ShapedText>().Elements)
 				{
-					EElementType ElementType = DrawElementMap.Key;
-					const FSlateDrawElementContainer& Container = DrawElementMap.Value;
-
-					//@TODO Support all element types and create their own JSON conversion function.
-					switch (ElementType)
-					{
-					case EElementType::ET_ShapedText:
-					{
-						for (const FSlateDrawElement& DrawElements : Container.Elements)
-						{
-							FSlateShapedTextPayload ShapedTextPayload = DrawElements.GetDataPayload<FSlateShapedTextPayload>();
-							FSlateTestHelper::ToJson(ShapedTextPayload, WriteJsonObject);
-						}
-					}
-					break;
-					default:
-						break;
-					}
+					FSlateShapedTextPayload ShapedTextPayload = DrawElements.GetDataPayload<FSlateShapedTextPayload>();
+					FSlateTestHelper::ToJson(ShapedTextPayload, WriteJsonObject);
 				}
 			}
 		}
