Fix missing includes CIS

diff --git a/Engine/Plugins/Tests/WidgetAutomationTests/Source/WidgetAutomationTests/Private/WidgetMockNonTemplate.cpp b/Engine/Plugins/Tests/WidgetAutomationTests/Source/WidgetAutomationTests/Private/WidgetMockNonTemplate.cpp
--- a/Engine/Plugins/Tests/WidgetAutomationTests/Source/WidgetAutomationTests/Private/WidgetMockNonTemplate.cpp
+++ b/Engine/Plugins/Tests/WidgetAutomationTests/Source/WidgetAutomationTests/Private/WidgetMockNonTemplate.cpp
@@ -2,6 +2,7 @@
 #include "WidgetMockNonTemplate.h"
 #include "Dom/JsonObject.h"
 #include "Misc/GeneratedTypeName.h"
+#include "Rendering/DrawElements.h"
 #include "Rendering/DrawElementPayloads.h"
 #include "Serialization/JsonReader.h"
 #include "SlateTestHelper.h"
